#ocn - Enable CrashReporter plugin by default on iOS

diff --git a/Engine/Config/Linux/LinuxEngine.ini b/Engine/Config/Linux/LinuxEngine.ini
--- a/Engine/Config/Linux/LinuxEngine.ini
+++ b/Engine/Config/Linux/LinuxEngine.ini
@@ -26,3 +26,6 @@ g.TimeoutForBlockOnRenderFence=60000
 
 [PlatformCrypto]
 PlatformRequiresDataCrypto=True
+
+[CrashReportClient]
+bAgreeToCrashUpload=true
\ No newline at end of file
