Remove the blanket agreement to upload crashes on Linux

(Edigrating CL 7251581 to 4.23)

diff --git a/Engine/Config/Linux/LinuxEngine.ini b/Engine/Config/Linux/LinuxEngine.ini
--- a/Engine/Config/Linux/LinuxEngine.ini
+++ b/Engine/Config/Linux/LinuxEngine.ini
@@ -26,6 +26,3 @@ g.TimeoutForBlockOnRenderFence=60000
 
 [PlatformCrypto]
 PlatformRequiresDataCrypto=True
-
-[CrashReportClient]
-bAgreeToCrashUpload=true
\ No newline at end of file
