UE4.21 - Remove Linux hack for Vulkan
#rnx

diff --git a/Engine/Config/Linux/LinuxEngine.ini b/Engine/Config/Linux/LinuxEngine.ini
--- a/Engine/Config/Linux/LinuxEngine.ini
+++ b/Engine/Config/Linux/LinuxEngine.ini
@@ -25,8 +25,5 @@ DeviceProfileSelectionModule="LinuxDeviceProfileSelector"
 ; larger timeout since drivers may take longer time
 g.TimeoutForBlockOnRenderFence=60000
 
-; Temp hack to unblock integration
-r.Vulkan.SubmitAfterEveryEndRenderPass=1
-
 [PlatformCrypto]
 PlatformRequiresDataCrypto=True
