Use the correct platform type (thanks moadib! PR #6242)

diff --git a/Engine/Config/Linux/DataDrivenPlatformInfo.ini b/Engine/Config/Linux/DataDrivenPlatformInfo.ini
--- a/Engine/Config/Linux/DataDrivenPlatformInfo.ini
+++ b/Engine/Config/Linux/DataDrivenPlatformInfo.ini
@@ -70,7 +70,7 @@ PlatformGroupName=Desktop
 [PlatformInfo LinuxServer]
 TargetPlatformName=LinuxServer
 DisplayName=Linux (Dedicated Server)
-PlatformType=Editor
+PlatformType=Server
 PlatformFlags=None
 NormalIconPath=Launcher/Linux/Platform_Linux_24x
 LargeIconPath=Launcher/Linux/Platform_Linux_128x
