Linux is enabled for Windows Editor when available

diff --git a/Engine/Config/Linux/DataDrivenPlatformInfo.ini b/Engine/Config/Linux/DataDrivenPlatformInfo.ini
--- a/Engine/Config/Linux/DataDrivenPlatformInfo.ini
+++ b/Engine/Config/Linux/DataDrivenPlatformInfo.ini
@@ -12,7 +12,7 @@ XLargeIconPath=
 UATCommandLine=
 AutoSDKPath=Linux_x64
 TutorialPath=/Platforms/Linux/GettingStarted
-Windows:bIsEnabled=false
+Windows:bIsEnabled=true
 Mac:bIsEnabled=false
 Linux:bIsEnabled=true
 BinariesDirectoryName=Linux
