UE5 - Fix Linux DDPI file for Vulkan
Also enabled VolumeTextureCompression and RGBColorBuffer in the inis (as the code was already using them)

diff --git a/Engine/Config/Linux/DataDrivenPlatformInfo.ini b/Engine/Config/Linux/DataDrivenPlatformInfo.ini
--- a/Engine/Config/Linux/DataDrivenPlatformInfo.ini
+++ b/Engine/Config/Linux/DataDrivenPlatformInfo.ini
@@ -22,37 +22,105 @@ bIsMetalMRT = false
 bIsPC = true
 bIsConsole = false
 bIsAndroidOpenGLES = false
-bSupportsMobileMultiView = false	; //#todo-vulkanddpi
-bSupportsArrayTextureCompression = false	; //#todo-vulkanddpi
-bSupportsVolumeTextureCompression = false	; //#todo-vulkanddpi
+
+; //#todo-vulkanddpi
+bSupportsMobileMultiView = false
+
+; //#todo-vulkanddpi
+bSupportsArrayTextureCompression = false
+
+bSupportsVolumeTextureCompression = true
 bSupportsDistanceFields = true
 bSupportsDiaphragmDOF = true
-bSupportsRGBColorBuffer = false	; //#todo-vulkanddpi
+bSupportsRGBColorBuffer = true
 bSupportsCapsuleShadows = true
 bSupportsVolumetricFog = true
 bSupportsIndexBufferUAVs = true
-bSupportsInstancedStereo = false	; //#todo-vulkanddpi
-bSupportsMultiView = false	; //#todo-vulkanddpi
+
+; //#todo-vulkanddpi
+bSupportsInstancedStereo = false
+
+; //#todo-vulkanddpi
+bSupportsMultiView = false
+
 bSupportsMSAA = true
-bSupports4ComponentUAVReadWrite = false	; //#todo-vulkanddpi
-bSupportsRenderTargetWriteMask = false	; //#todo-vulkanddpi
+
+; //#todo-vulkanddpi
+bSupports4ComponentUAVReadWrite = false
+
+; //#todo-vulkanddpi
+bSupportsRenderTargetWriteMask = false
+
 bSupportsRayTracing = false
 bSupportsRayTracingIndirectInstanceData = false
 bSupportsGPUSkinCache = true
-bSupportsByteBufferComputeShaders = false	; //#todo-vulkanddpi
+
+; //#todo-vulkanddpi
+bSupportsByteBufferComputeShaders = false
+
 bSupportsGPUScene = true
-bSupportsPrimitiveShaders = false	; //#todo-vulkanddpi
-bSupportsUInt64ImageAtomics = false	; //#todo-vulkanddpi
+
+; //#todo-vulkanddpi
+bSupportsPrimitiveShaders = false
+
+; //#todo-vulkanddpi
+bSupportsUInt64ImageAtomics = false
+
+; //#todo-vulkanddpi
+bRequiresVendorExtensionsForAtomics = false
+
+; //#todo-vulkanddpi
 bSupportsNanite = false
+
+; //#todo-vulkanddpi
 bSupportsLumenGI = false
-bSupportsSSDIndirect = false	; //#todo-vulkanddpi
-bSupportsTemporalHistoryUpscale = false	; //#todo-vulkanddpi
-bSupportsRTIndexFromVS = false	; //#todo-vulkanddpi
-bSupportsIntrinsicWaveOnce = false	; //#todo-vulkanddpi
-bSupportsConservativeRasterization = false	; //#todo-vulkanddpi
-bSupportsWaveOperations = false	; //#todo-vulkanddpi
-bRequiresExplicit128bitRT = false	; //#todo-vulkanddpi
-bSupportsGen5TemporalAA = false	; //#todo-vulkanddpi
+
+; //#todo-vulkanddpi
+bSupportsSSDIndirect = false
+
+; //#todo-vulkanddpi
+bSupportsTemporalHistoryUpscale = false
+
+; //#todo-vulkanddpi
+bSupportsRTIndexFromVS = false
+
+; //#todo-vulkanddpi
+bSupportsIntrinsicWaveOnce = false
+
+; //#todo-vulkanddpi
+bSupportsConservativeRasterization = false
+
+; //#todo-vulkanddpi
+bSupportsWaveOperations = false
+
+; //#todo-vulkanddpi
+bRequiresExplicit128bitRT = false
+
+; //#todo-vulkanddpi
+bSupportsGen5TemporalAA = false
+
 bTargetsTiledGPU = false
 bNeedsOfflineCompiler = false
-bSupportsComputeFramework = false	; //#todo-vulkanddpi
+
+; //#todo-vulkanddpi
+bSupportsComputeFramework = false
+
+; //#todo-vulkanddpi
+bSupportsDualSourceBlending = false
+
+; //#todo-vulkanddpi
+bRequiresGeneratePrevTransformBuffer = false
+
+bRequiresRenderTargetDuringRaster = true
+
+; //#todo-vulkanddpi
+bRequiresDisableForwardLocalLights = true
+
+bCompileSignalProcessingPipeline = true
+bSupportsTessellation = true
+
+; //#todo-vulkanddpi
+bSupportsPerPixelDBufferMask = false
+
+bIsHlslcc = false
+NumberOfComputeThreads = 64
diff --git a/Engine/Config/Windows/DataDrivenPlatformInfo.ini b/Engine/Config/Windows/DataDrivenPlatformInfo.ini
--- a/Engine/Config/Windows/DataDrivenPlatformInfo.ini
+++ b/Engine/Config/Windows/DataDrivenPlatformInfo.ini
@@ -69,15 +69,10 @@ bSupportsMobileMultiView = false
 ; //#todo-vulkanddpi
 bSupportsArrayTextureCompression = false
 
-; //#todo-vulkanddpi
-bSupportsVolumeTextureCompression = false
-
+bSupportsVolumeTextureCompression = true
 bSupportsDistanceFields = true
 bSupportsDiaphragmDOF = true
-
-; //#todo-vulkanddpi
-bSupportsRGBColorBuffer = false
-
+bSupportsRGBColorBuffer = true
 bSupportsCapsuleShadows = true
 bSupportsVolumetricFog = true
 bSupportsIndexBufferUAVs = true
diff --git a/Engine/Source/Runtime/Engine/Private/VolumeTexture.cpp b/Engine/Source/Runtime/Engine/Private/VolumeTexture.cpp
--- a/Engine/Source/Runtime/Engine/Private/VolumeTexture.cpp
+++ b/Engine/Source/Runtime/Engine/Private/VolumeTexture.cpp
@@ -383,7 +383,6 @@ bool UVolumeTexture::ShaderPlatformSupportsCompression(FStaticShaderPlatform Sha
 	switch (ShaderPlatform)
 	{
 	case SP_PCD3D_SM5:
-	case SP_VULKAN_SM5:
 	case SP_VULKAN_SM5_LUMIN:
 		return true;
 
diff --git a/Engine/Source/Runtime/Renderer/Private/PostProcess/DiaphragmDOF.cpp b/Engine/Source/Runtime/Renderer/Private/PostProcess/DiaphragmDOF.cpp
--- a/Engine/Source/Runtime/Renderer/Private/PostProcess/DiaphragmDOF.cpp
+++ b/Engine/Source/Runtime/Renderer/Private/PostProcess/DiaphragmDOF.cpp
@@ -681,9 +681,7 @@ static FORCEINLINE bool SupportsRGBColorBuffer(EShaderPlatform ShaderPlatform)
 	}
 
 	// There is high number of UAV to write in reduce pass.
-	return 
-		ShaderPlatform == SP_VULKAN_SM5 ||
-		FDataDrivenShaderPlatformInfo::GetSupportsRGBColorBuffer(ShaderPlatform);
+	return FDataDrivenShaderPlatformInfo::GetSupportsRGBColorBuffer(ShaderPlatform);
 }
 
 
