Removes a duplicated bSupportsGen5TemporalAA in linux's DataDrivenPlatformInfo.ini

diff --git a/Engine/Config/Linux/DataDrivenPlatformInfo.ini b/Engine/Config/Linux/DataDrivenPlatformInfo.ini
--- a/Engine/Config/Linux/DataDrivenPlatformInfo.ini
+++ b/Engine/Config/Linux/DataDrivenPlatformInfo.ini
@@ -100,9 +100,6 @@ bSupportsWaveOperations = false
 ; //#todo-vulkanddpi
 bRequiresExplicit128bitRT = false
 
-; //#todo-vulkanddpi
-bSupportsGen5TemporalAA = false
-
 bTargetsTiledGPU = false
 bNeedsOfflineCompiler = false
 
