Fix metasounds ensure PR #8404 (Thanks slonopotamus!)

diff --git a/Engine/Config/Linux/DataDrivenPlatformInfo.ini b/Engine/Config/Linux/DataDrivenPlatformInfo.ini
--- a/Engine/Config/Linux/DataDrivenPlatformInfo.ini
+++ b/Engine/Config/Linux/DataDrivenPlatformInfo.ini
@@ -23,6 +23,7 @@ Linux:bUsesHostCompiler=true
 bUATClosesAfterLaunch=true
 PlatformGroupName=Desktop
 PlatformSubMenu=Linux
+AudioCompressionSettingsIniSectionName=/Script/LinuxTargetPlatform.LinuxTargetSettings
 
 [ShaderPlatform VULKAN_SM5]
 bIsMobile = false
diff --git a/Engine/Config/LinuxArm64/DataDrivenPlatformInfo.ini b/Engine/Config/LinuxArm64/DataDrivenPlatformInfo.ini
--- a/Engine/Config/LinuxArm64/DataDrivenPlatformInfo.ini
+++ b/Engine/Config/LinuxArm64/DataDrivenPlatformInfo.ini
@@ -24,4 +24,4 @@ Linux:bUsesHostCompiler=true
 bUATClosesAfterLaunch=true
 PlatformGroupName=Mobile
 PlatformSubMenu=Linux
-
+AudioCompressionSettingsIniSectionName=/Script/LinuxTargetPlatform.LinuxTargetSettings
