Fix up crash reporter config-var project overrides not working

diff --git a/Engine/Config/Linux/BaseLinuxEngine.ini b/Engine/Config/Linux/BaseLinuxEngine.ini
--- a/Engine/Config/Linux/BaseLinuxEngine.ini
+++ b/Engine/Config/Linux/BaseLinuxEngine.ini
@@ -6,3 +6,6 @@ UnrealBuildTool=Engine/Build/BatchFiles/Linux/Build.sh
 
 [OnlineServices.EOS.Auth]
 DefaultExternalCredentialTypeStr=Steam
+
+[CrashReportClient]
+bAgreeToCrashUpload=false
diff --git a/Engine/Config/Linux/LinuxEngine.ini b/Engine/Config/Linux/LinuxEngine.ini
--- a/Engine/Config/Linux/LinuxEngine.ini
+++ b/Engine/Config/Linux/LinuxEngine.ini
@@ -27,6 +27,3 @@ g.TimeoutForBlockOnRenderFence=60000
 
 [PlatformCrypto]
 PlatformRequiresDataCrypto=True
-
-[CrashReportClient]
-bAgreeToCrashUpload=false
diff --git a/Engine/Config/LinuxArm64/BaseLinuxArm64Engine.ini b/Engine/Config/LinuxArm64/BaseLinuxArm64Engine.ini
new file mode 100644
--- /dev/null
+++ b/Engine/Config/LinuxArm64/BaseLinuxArm64Engine.ini
@@ -0,0 +1,2 @@
+[CrashReportClient]
+bAgreeToCrashUpload=false
diff --git a/Engine/Config/LinuxArm64/LinuxArm64Engine.ini b/Engine/Config/LinuxArm64/LinuxArm64Engine.ini
--- a/Engine/Config/LinuxArm64/LinuxArm64Engine.ini
+++ b/Engine/Config/LinuxArm64/LinuxArm64Engine.ini
@@ -28,6 +28,3 @@ g.TimeoutForBlockOnRenderFence=60000
 
 [PlatformCrypto]
 PlatformRequiresDataCrypto=True
-
-[CrashReportClient]
-bAgreeToCrashUpload=false
