VulkanRHI: Enable bindless for Linux in SM6.

#rnx

diff --git a/Engine/Config/Linux/BaseLinuxEngine.ini b/Engine/Config/Linux/BaseLinuxEngine.ini
--- a/Engine/Config/Linux/BaseLinuxEngine.ini
+++ b/Engine/Config/Linux/BaseLinuxEngine.ini
@@ -10,11 +10,10 @@ DefaultExternalCredentialTypeStr=Steam
 [CrashReportClient]
 bAgreeToCrashUpload=false
 
-; To enable experimental Vulkan ray tracing, set both of these values to Enabled.
-; NOTE: Must also be enabled in VulkanPC/DataDrivenPlatformInfo.ini
-;[SF_VULKAN_SM6]
-;BindlessResources=Enabled
-;BindlessSamplers=Enabled
+[SF_VULKAN_SM6]
+; Bindless is necessary for Vulkan ray tracing
+BindlessResources=Enabled
+BindlessSamplers=Enabled
 
 [ConsoleVariables]
 ; necessary for LowMemoryWarning events to be fired
