VulkanRHI: default value for unused cvar in Vulkan generated different DDC keys for Linux shaders.

#rnx

diff --git a/Engine/Config/Linux/LinuxEngine.ini b/Engine/Config/Linux/LinuxEngine.ini
--- a/Engine/Config/Linux/LinuxEngine.ini
+++ b/Engine/Config/Linux/LinuxEngine.ini
@@ -27,3 +27,6 @@ g.TimeoutForBlockOnRenderFence=60000
 
 [PlatformCrypto]
 PlatformRequiresDataCrypto=True
+
+[/Script/Engine.RendererSettings]
+r.Shaders.RemoveUnusedInterpolators=1
