Fix missing character in Linux_SDK.json

diff --git a/Engine/Config/Linux/Linux_SDK.json b/Engine/Config/Linux/Linux_SDK.json
--- a/Engine/Config/Linux/Linux_SDK.json
+++ b/Engine/Config/Linux/Linux_SDK.json
@@ -1,7 +1,7 @@
 {
-"MainVersion" : "v23_clang-18.1.0-rockylinux8",
-"MinVersion" : "v23_clang-18.1.0-rockylinux8",
-"MaxVersion" : "v23_clang-18.1.0-rockylinux8",
+"MainVersion" : "v22_clang-16.0.6-centos7",
+"MinVersion" : "v22_clang-16.0.6-centos7",
+"MaxVersion" : "v22_clang-16.0.6-centos7",
 
 "AutoSDKPlatform" : "Linux_x64",
 }
