Move Substrate settings to RendererSettings category to make the overridable per project.

This was causing some project like SubstateTest to no longer have opaque rough refraction enabled.

diff --git a/Engine/Config/Linux/BaseLinuxEngine.ini b/Engine/Config/Linux/BaseLinuxEngine.ini
--- a/Engine/Config/Linux/BaseLinuxEngine.ini
+++ b/Engine/Config/Linux/BaseLinuxEngine.ini
@@ -11,7 +11,7 @@ bAgreeToCrashUpload=false
 ; Bindless is necessary for Vulkan ray tracing
 BindlessConfiguration=All
 
-[SystemSettings]
+[/Script/Engine.RendererSettings]
 ; Substrate settings for more advanced shading on PC and editor by default.
 r.Substrate.RoughDiffuse=1
 r.Substrate.SheenQuality=1
diff --git a/Engine/Config/LinuxArm64/BaseLinuxArm64Engine.ini b/Engine/Config/LinuxArm64/BaseLinuxArm64Engine.ini
--- a/Engine/Config/LinuxArm64/BaseLinuxArm64Engine.ini
+++ b/Engine/Config/LinuxArm64/BaseLinuxArm64Engine.ini
@@ -1,8 +1,7 @@
 [CrashReportClient]
 bAgreeToCrashUpload=false
 
-
-[SystemSettings]
+[/Script/Engine.RendererSettings]
 ; Substrate settings for more advanced shading on PC and editor by default.
 r.Substrate.RoughDiffuse=1
 r.Substrate.SheenQuality=1
