Remove OGG as default codec on Linux to match all other targets.

#tests
#DisabledValidations Preflight

diff --git a/Engine/Config/Linux/LinuxEngine.ini b/Engine/Config/Linux/LinuxEngine.ini
--- a/Engine/Config/Linux/LinuxEngine.ini
+++ b/Engine/Config/Linux/LinuxEngine.ini
@@ -1,8 +1,6 @@
 [Audio]
 ; Defines a platform-specific volume headroom (in dB) for audio to provide better platform consistency with respect to volume levels.
 PlatformHeadroomDB=0
-PlatformFormat=OGG
-PlatformStreamingFormat=OGG
 
 [SystemSettings]
 r.setres=1280x720
